LiveLinkOpenVR: Marshal SteamVR Input into LiveLinkGamepadInputDevice role.

#rnx

diff --git a/Engine/Plugins/Experimental/LiveLinkOpenVR/Config/livelinkopenvr_action_manifest.json b/Engine/Plugins/Experimental/LiveLinkOpenVR/Config/livelinkopenvr_action_manifest.json
new file mode 100644
--- /dev/null
+++ b/Engine/Plugins/Experimental/LiveLinkOpenVR/Config/livelinkopenvr_action_manifest.json
@@ -0,0 +1,131 @@
+{
+	"default_bindings": [
+		{
+			"controller_type": "knuckles",
+			"binding_url": "livelinkopenvr_bindings_knuckles.json"
+		}
+	],
+	"actions": [
+		{
+			"name": "/actions/LiveLinkGamepadInputDevice/in/LeftAnalog_2D",
+			"type": "vector2"
+		},
+		{
+			"name": "/actions/LiveLinkGamepadInputDevice/in/RightAnalog_2D",
+			"type": "vector2"
+		},
+		{
+			"name": "/actions/LiveLinkGamepadInputDevice/in/LeftTriggerAnalog",
+			"type": "vector1"
+		},
+		{
+			"name": "/actions/LiveLinkGamepadInputDevice/in/RightTriggerAnalog",
+			"type": "vector1"
+		},
+		{
+			"name": "/actions/LiveLinkGamepadInputDevice/in/LeftThumb",
+			"type": "vector1"
+		},
+		{
+			"name": "/actions/LiveLinkGamepadInputDevice/in/RightThumb",
+			"type": "vector1"
+		},
+		{
+			"name": "/actions/LiveLinkGamepadInputDevice/in/SpecialLeft",
+			"type": "vector1"
+		},
+		{
+			"name": "/actions/LiveLinkGamepadInputDevice/in/SpecialLeft_2D",
+			"type": "vector2"
+		},
+		{
+			"name": "/actions/LiveLinkGamepadInputDevice/in/SpecialRight",
+			"type": "vector1"
+		},
+		{
+			"name": "/actions/LiveLinkGamepadInputDevice/in/FaceButtonBottom",
+			"type": "vector1"
+		},
+		{
+			"name": "/actions/LiveLinkGamepadInputDevice/in/FaceButtonRight",
+			"type": "vector1"
+		},
+		{
+			"name": "/actions/LiveLinkGamepadInputDevice/in/FaceButtonLeft",
+			"type": "vector1"
+		},
+		{
+			"name": "/actions/LiveLinkGamepadInputDevice/in/FaceButtonTop",
+			"type": "vector1"
+		},
+		{
+			"name": "/actions/LiveLinkGamepadInputDevice/in/LeftShoulder",
+			"type": "vector1"
+		},
+		{
+			"name": "/actions/LiveLinkGamepadInputDevice/in/RightShoulder",
+			"type": "vector1"
+		},
+		{
+			"name": "/actions/LiveLinkGamepadInputDevice/in/LeftTriggerThreshold",
+			"type": "vector1"
+		},
+		{
+			"name": "/actions/LiveLinkGamepadInputDevice/in/RightTriggerThreshold",
+			"type": "vector1"
+		},
+		{
+			"name": "/actions/LiveLinkGamepadInputDevice/in/DPadUp",
+			"type": "vector1"
+		},
+		{
+			"name": "/actions/LiveLinkGamepadInputDevice/in/DPadDown",
+			"type": "vector1"
+		},
+		{
+			"name": "/actions/LiveLinkGamepadInputDevice/in/DPadRight",
+			"type": "vector1"
+		},
+		{
+			"name": "/actions/LiveLinkGamepadInputDevice/in/DPadLeft",
+			"type": "vector1"
+		},
+		{
+			"name": "/actions/LiveLinkGamepadInputDevice/in/LeftStickUp",
+			"type": "vector1"
+		},
+		{
+			"name": "/actions/LiveLinkGamepadInputDevice/in/LeftStickDown",
+			"type": "vector1"
+		},
+		{
+			"name": "/actions/LiveLinkGamepadInputDevice/in/LeftStickRight",
+			"type": "vector1"
+		},
+		{
+			"name": "/actions/LiveLinkGamepadInputDevice/in/LeftStickLeft",
+			"type": "vector1"
+		},
+		{
+			"name": "/actions/LiveLinkGamepadInputDevice/in/RightStickUp",
+			"type": "vector1"
+		},
+		{
+			"name": "/actions/LiveLinkGamepadInputDevice/in/RightStickDown",
+			"type": "vector1"
+		},
+		{
+			"name": "/actions/LiveLinkGamepadInputDevice/in/RightStickRight",
+			"type": "vector1"
+		},
+		{
+			"name": "/actions/LiveLinkGamepadInputDevice/in/RightStickLeft",
+			"type": "vector1"
+		}
+	],
+	"localization": [
+		{
+			"language_tag": "en_US"
+		}
+	]
+}
diff --git a/Engine/Plugins/Experimental/LiveLinkOpenVR/Config/livelinkopenvr_bindings_knuckles.json b/Engine/Plugins/Experimental/LiveLinkOpenVR/Config/livelinkopenvr_bindings_knuckles.json
new file mode 100644
--- /dev/null
+++ b/Engine/Plugins/Experimental/LiveLinkOpenVR/Config/livelinkopenvr_bindings_knuckles.json
@@ -0,0 +1,65 @@
+{
+	"controller_type": "knuckles",
+	"name": "LiveLinkOpenVR bindings for Valve Index controller",
+	"description": "Bindings for the Live Link OpenVR plugin for the Valve Index controller",
+	"bindings": {
+		"/actions/LiveLinkGamepadInputDevice": {
+			"sources": [
+				{
+					"path": "/user/hand/left/input/thumbstick",
+					"mode": "joystick",
+					"inputs": {
+						"position": {
+							"output": "/actions/LiveLinkGamepadInputDevice/in/LeftAnalog_2D"
+						},
+						"click": {
+							"output": "/actions/LiveLinkGamepadInputDevice/in/LeftThumb"
+						}
+					}
+				},
+				{
+					"path": "/user/hand/left/input/trigger",
+					"mode": "trigger",
+					"inputs": {
+						"pull": {
+							"output": "/actions/LiveLinkGamepadInputDevice/in/LeftTriggerAnalog"
+						}
+					}
+				},
+				{
+					"path": "/user/hand/left/input/trackpad",
+					"mode": "trackpad",
+					"inputs": {
+						"position": {
+							"output": "/actions/LiveLinkGamepadInputDevice/in/SpecialLeft_2D"
+						},
+						"click": {
+							"output": "/actions/LiveLinkGamepadInputDevice/in/SpecialLeft"
+						}
+					}
+				},
+				{
+					"path": "/user/hand/right/input/thumbstick",
+					"mode": "joystick",
+					"inputs": {
+						"position": {
+							"output": "/actions/LiveLinkGamepadInputDevice/in/RightAnalog_2D"
+						},
+						"click": {
+							"output": "/actions/LiveLinkGamepadInputDevice/in/RightThumb"
+						}
+					}
+				},
+				{
+					"path": "/user/hand/right/input/trigger",
+					"mode": "trigger",
+					"inputs": {
+						"pull": {
+							"output": "/actions/LiveLinkGamepadInputDevice/in/RightTriggerAnalog"
+						}
+					}
+				}
+			]
+		}
+	}
+}
diff --git a/Engine/Plugins/Experimental/LiveLinkOpenVR/Source/LiveLinkOpenVR/Private/LiveLinkOpenVRModule.cpp b/Engine/Plugins/Experimental/LiveLinkOpenVR/Source/LiveLinkOpenVR/Private/LiveLinkOpenVRModule.cpp
--- a/Engine/Plugins/Experimental/LiveLinkOpenVR/Source/LiveLinkOpenVR/Private/LiveLinkOpenVRModule.cpp
+++ b/Engine/Plugins/Experimental/LiveLinkOpenVR/Source/LiveLinkOpenVR/Private/LiveLinkOpenVRModule.cpp
@@ -1,6 +1,7 @@
 // Copyright Epic Games, Inc. All Rights Reserved.
 
 #include "LiveLinkOpenVRModule.h"
+#include "HAL/FileManager.h"
 #include "HAL/PlatformProcess.h"
 #include "Interfaces/IPluginManager.h"
 #include "Logging/StructuredLog.h"
@@ -44,8 +45,9 @@ vr::IVRSystem* FLiveLinkOpenVRModule::GetVrSystem()
 
 bool FLiveLinkOpenVRModule::LoadOpenVRLibrary()
 {
-	const TCHAR* const OpenVRSdkVer = TEXT("OpenVRv1_5_17");
 	const FString PluginBaseDir = IPluginManager::Get().FindPlugin(TEXT("LiveLinkOpenVR"))->GetBaseDir();
+
+	const TCHAR* const OpenVRSdkVer = TEXT("OpenVRv1_5_17");
 	FString OpenVRSdkRoot = FString::Printf(TEXT("%s/Source/ThirdParty/OpenVR/%s"), *PluginBaseDir, OpenVRSdkVer);
 
 #if PLATFORM_WINDOWS
@@ -71,7 +73,7 @@ bool FLiveLinkOpenVRModule::LoadOpenVRLibrary()
 
 	if (!OpenVRDLLHandle)
 	{
-		UE_LOG(LogLiveLinkOpenVR, Log, TEXT("Failed to load OpenVR library."));
+		UE_LOGFMT(LogLiveLinkOpenVR, Log, "Failed to load OpenVR library.");
 		return false;
 	}
 
@@ -85,6 +87,14 @@ bool FLiveLinkOpenVRModule::LoadOpenVRLibrary()
 		return false;
 	}
 
+	FString ManifestPath = FPaths::Combine(PluginBaseDir, "Config", "livelinkopenvr_action_manifest.json");
+	ManifestPath = IFileManager::Get().ConvertToAbsolutePathForExternalAppForRead(*ManifestPath);
+	const vr::EVRInputError InputError = vr::VRInput()->SetActionManifestPath(TCHAR_TO_UTF8(*ManifestPath));
+	if (InputError != vr::EVRInputError::VRInputError_None)
+	{
+		UE_LOGFMT(LogLiveLinkOpenVR, Error, "IVRInput::SetActionManifestPath failed with result {InputError}", InputError);
+	}
+
 	return true;
 }
 
diff --git a/Engine/Plugins/Experimental/LiveLinkOpenVR/Source/LiveLinkOpenVR/Private/LiveLinkOpenVRSource.cpp b/Engine/Plugins/Experimental/LiveLinkOpenVR/Source/LiveLinkOpenVR/Private/LiveLinkOpenVRSource.cpp
--- a/Engine/Plugins/Experimental/LiveLinkOpenVR/Source/LiveLinkOpenVR/Private/LiveLinkOpenVRSource.cpp
+++ b/Engine/Plugins/Experimental/LiveLinkOpenVR/Source/LiveLinkOpenVR/Private/LiveLinkOpenVRSource.cpp
@@ -6,7 +6,10 @@
 #include "LiveLinkOpenVRModule.h"
 #include "LiveLinkSubjectSettings.h"
 #include "Logging/StructuredLog.h"
+#include "Misc/App.h"
 #include "Misc/CoreDelegates.h"
+#include "Roles/LiveLinkInputDeviceRole.h"
+#include "Roles/LiveLinkInputDeviceTypes.h"
 #include "Roles/LiveLinkTransformRole.h"
 #include "Roles/LiveLinkTransformTypes.h"
 
@@ -95,15 +98,15 @@ void FLiveLinkOpenVRSource::ReceiveClient(ILiveLinkClient* InClient, FGuid InSou
 }
 
 
-void FLiveLinkOpenVRSource::InitializeSettings(ULiveLinkSourceSettings* Settings)
+void FLiveLinkOpenVRSource::InitializeSettings(ULiveLinkSourceSettings* InSettings)
 {
-	ULiveLinkOpenVRSourceSettings* SourceSettings = Cast<ULiveLinkOpenVRSourceSettings>(Settings);
+	ULiveLinkOpenVRSourceSettings* SourceSettings = Cast<ULiveLinkOpenVRSourceSettings>(InSettings);
 	if (!ensure(SourceSettings))
 	{
 		return;
 	}
 
-	LocalUpdateRateInHz = SourceSettings->CommonSettings.LocalUpdateRateInHz;
+	LocalUpdateRateInHz_AnyThread = SourceSettings->CommonSettings.LocalUpdateRateInHz;
 }
 
 
@@ -152,26 +155,129 @@ void FLiveLinkOpenVRSource::Stop()
 }
 
 
+struct FOpenVRInputAction
+{
+	enum EActionType
+	{
+		Digital,
+		Analog1D,
+		Analog2D,
+	};
+
+	FName ActionName;
+	EActionType ActionType;
+	TOptional<float FLiveLinkGamepadInputDeviceFrameData::*> OutputFieldX;
+	TOptional<float FLiveLinkGamepadInputDeviceFrameData::*> OutputFieldY;
+
+	vr::VRActionHandle_t Handle = vr::k_ulInvalidActionHandle;
+
+	union {
+		vr::InputDigitalActionData_t Digital;
+		vr::InputAnalogActionData_t Analog;
+	} LastActionData;
+};
+
+
 uint32 FLiveLinkOpenVRSource::Run()
 {
+	TArray<FOpenVRInputAction> Actions = {
+		{ "LeftAnalog_2D", FOpenVRInputAction::Analog2D,
+		  &FLiveLinkGamepadInputDeviceFrameData::LeftAnalogX, &FLiveLinkGamepadInputDeviceFrameData::LeftAnalogY },
+
+		{ "RightAnalog_2D", FOpenVRInputAction::Analog2D,
+		  &FLiveLinkGamepadInputDeviceFrameData::RightAnalogX, &FLiveLinkGamepadInputDeviceFrameData::RightAnalogY },
+
+		{ "SpecialLeft_2D", FOpenVRInputAction::Analog2D,
+		  &FLiveLinkGamepadInputDeviceFrameData::SpecialLeft_X, &FLiveLinkGamepadInputDeviceFrameData::SpecialLeft_Y },
+
+#define DECLARE_VECTOR1_ACTION(ActionNameToken)                     \
+		{ UE_STRINGIZE(ActionNameToken), FOpenVRInputAction::Analog1D, \
+		  &FLiveLinkGamepadInputDeviceFrameData::ActionNameToken },
+
+		//DECLARE_VECTOR1_ACTION(LeftAnalogX)
+		//DECLARE_VECTOR1_ACTION(LeftAnalogY)
+		//DECLARE_VECTOR1_ACTION(RightAnalogX)
+		//DECLARE_VECTOR1_ACTION(RightAnalogY)
+		DECLARE_VECTOR1_ACTION(LeftTriggerAnalog)
+		DECLARE_VECTOR1_ACTION(RightTriggerAnalog)
+		DECLARE_VECTOR1_ACTION(LeftThumb)
+		DECLARE_VECTOR1_ACTION(RightThumb)
+		DECLARE_VECTOR1_ACTION(SpecialLeft)
+		//DECLARE_VECTOR1_ACTION(SpecialLeft_X)
+		//DECLARE_VECTOR1_ACTION(SpecialLeft_Y)
+		DECLARE_VECTOR1_ACTION(SpecialRight)
+		DECLARE_VECTOR1_ACTION(FaceButtonBottom)
+		DECLARE_VECTOR1_ACTION(FaceButtonRight)
+		DECLARE_VECTOR1_ACTION(FaceButtonLeft)
+		DECLARE_VECTOR1_ACTION(FaceButtonTop)
+		DECLARE_VECTOR1_ACTION(LeftShoulder)
+		DECLARE_VECTOR1_ACTION(RightShoulder)
+		DECLARE_VECTOR1_ACTION(LeftTriggerThreshold)
+		DECLARE_VECTOR1_ACTION(RightTriggerThreshold)
+		DECLARE_VECTOR1_ACTION(DPadUp)
+		DECLARE_VECTOR1_ACTION(DPadDown)
+		DECLARE_VECTOR1_ACTION(DPadRight)
+		DECLARE_VECTOR1_ACTION(DPadLeft)
+		DECLARE_VECTOR1_ACTION(LeftStickUp)
+		DECLARE_VECTOR1_ACTION(LeftStickDown)
+		DECLARE_VECTOR1_ACTION(LeftStickRight)
+		DECLARE_VECTOR1_ACTION(LeftStickLeft)
+		DECLARE_VECTOR1_ACTION(RightStickUp)
+		DECLARE_VECTOR1_ACTION(RightStickDown)
+		DECLARE_VECTOR1_ACTION(RightStickRight)
+		DECLARE_VECTOR1_ACTION(RightStickLeft)
+
+#undef DECLARE_VECTOR1_ACTION
+	};
+
 	FLiveLinkOpenVRModule& Module = FLiveLinkOpenVRModule::Get();
+	vr::IVRSystem* VrSystem = Module.GetVrSystem();
+	vr::IVRInput* VrInput = vr::VRInput();
+
+	const char* const ActionSetPath = "/actions/LiveLinkGamepadInputDevice";
+	vr::VRActionSetHandle_t ActionSet = vr::k_ulInvalidActionSetHandle;
+	vr::EVRInputError InputError = VrInput->GetActionSetHandle(ActionSetPath, &ActionSet);
+	if (InputError != vr::VRInputError_None)
+	{
+		UE_LOGFMT(LogLiveLinkOpenVR, Error, "IVRInput::GetActionSetHandle failed with result {InputError}", InputError);
+		ActionSet = vr::k_ulInvalidActionSetHandle;
+	}
+	else
+	{
+		for (FOpenVRInputAction& Action : Actions)
+		{
+			constexpr int32 MaxPathLen = 256;
+			TAnsiStringBuilder<MaxPathLen> ActionPath;
+			ActionPath.Appendf("%s/in/%S", ActionSetPath, *Action.ActionName.ToString());
+			InputError = VrInput->GetActionHandle(*ActionPath, &Action.Handle);
+			if (InputError != vr::VRInputError_None)
+			{
+				UE_LOGFMT(LogLiveLinkOpenVR, Error, "IVRInput::GetActionHandle for '{ActionPath}' failed with result {InputError}", *ActionPath, InputError);
+				Action.Handle = vr::k_ulInvalidActionHandle;
+			}
+		}
+	}
+
+	static const FName InputSubjectName("OpenVRInput");
+	Client->PushSubjectStaticData_AnyThread({ SourceGuid, InputSubjectName }, ULiveLinkInputDeviceRole::StaticClass(),
+		FLiveLinkStaticDataStruct(FLiveLinkGamepadInputDeviceStaticData::StaticStruct()));
 
 	TStaticArray<vr::TrackedDevicePose_t, vr::k_unMaxTrackedDeviceCount> Poses;
 	TStringBuilder<256> StringBuilder;
 
-	TMap<FName, FTransform> SubjectPoses;
 	double LastFrameTimeSec = -DBL_MAX;
 	while (!bStopping)
 	{
 		// Send new poses at the user specified update rate
-		const double FrameIntervalSec = 1.0 / LocalUpdateRateInHz;
+		const double FrameIntervalSec = 1.0 / LocalUpdateRateInHz_AnyThread;
 		const double TimeNowSec = FPlatformTime::Seconds();
 		if (TimeNowSec >= (LastFrameTimeSec + FrameIntervalSec))
 		{
 			LastFrameTimeSec = TimeNowSec;
-			SubjectPoses.Reset();
 
-			vr::IVRSystem* VrSystem = Module.GetVrSystem();
+			const TOptional<FQualifiedFrameTime> CurrentFrameTime = FApp::GetCurrentFrameTime();
+
+			// Update poses.
 			VrSystem->GetDeviceToAbsoluteTrackingPose(
 				vr::ETrackingUniverseOrigin::TrackingUniverseStanding,
 				0.0f,
@@ -179,6 +285,12 @@ uint32 FLiveLinkOpenVRSource::Run()
 				Poses.Num()
 			);
 
+			// Update inputs.
+			vr::VRActiveActionSet_t ActiveSet = { 0 };
+			ActiveSet.ulActionSet = ActionSet;
+			VrInput->UpdateActionState(&ActiveSet, sizeof(ActiveSet), 1);
+
+			// Enumerate poses.
 			for (int32 DeviceIdx = 0; DeviceIdx < Poses.Num(); ++DeviceIdx)
 			{
 				const vr::TrackedDevicePose_t& Pose = Poses[DeviceIdx];
@@ -197,11 +309,45 @@ uint32 FLiveLinkOpenVRSource::Run()
 					const vr::ETrackedDeviceClass DeviceClass = VrSystem->GetTrackedDeviceClass(DeviceIdx);
 					switch (DeviceClass)
 					{
-						case vr::TrackedDeviceClass_HMD:               StringBuilder << TEXT("HMD"); break;
-						case vr::TrackedDeviceClass_Controller:        StringBuilder << TEXT("Controller"); break;
-						case vr::TrackedDeviceClass_GenericTracker:    StringBuilder << TEXT("Tracker"); break;
-						case vr::TrackedDeviceClass_TrackingReference: StringBuilder << TEXT("TrackingRef"); break;
-						default:                                       StringBuilder << TEXT("Other"); break;
+						case vr::TrackedDeviceClass_HMD:
+							if (!ConnectionSettings.bTrackHMDs)
+							{
+								continue;
+							}
+
+							StringBuilder << TEXT("HMD");
+							break;
+
+						case vr::TrackedDeviceClass_Controller:
+							if (!ConnectionSettings.bTrackControllers)
+							{
+								continue;
+							}
+
+							StringBuilder << TEXT("Controller");
+							break;
+
+						case vr::TrackedDeviceClass_GenericTracker:
+							if (!ConnectionSettings.bTrackTrackers)
+							{
+								continue;
+							}
+
+							StringBuilder << TEXT("Tracker");
+							break;
+
+						case vr::TrackedDeviceClass_TrackingReference:
+							if (!ConnectionSettings.bTrackTrackingReferences)
+							{
+								continue;
+							}
+
+							StringBuilder << TEXT("TrackingRef");
+							break;
+
+						default:
+							StringBuilder << TEXT("Other");
+							break;
 					}
 
 					StringBuilder << TEXT("_");
@@ -223,29 +369,115 @@ uint32 FLiveLinkOpenVRSource::Run()
 					Client->PushSubjectStaticData_AnyThread({SourceGuid, SubjectName}, ULiveLinkTransformRole::StaticClass(), MoveTemp(StaticData));
 				}
 
-				// We might have static data, but not frame data.
-				if (!Pose.bPoseIsValid)
+				// Send transform frame data, if available.
+				if (Pose.bPoseIsValid)
+				{
+					// Transpose and decompose.
+					const FMatrix PoseMatrix = ToFMatrix(Pose.mDeviceToAbsoluteTracking);
+					const FQuat PoseOrientation(PoseMatrix);
+					const FVector PosePosition(PoseMatrix.M[3][0], PoseMatrix.M[3][1], PoseMatrix.M[3][2]);
+
+					// Handedness/basis change + scale.
+					const double MetersToUnrealUnits = 100.0; // cm
+					const FTransform PoseTransform(
+						FQuat(-PoseOrientation.Z, PoseOrientation.X, PoseOrientation.Y, -PoseOrientation.W),
+						FVector(-PosePosition.Z, PosePosition.X, PosePosition.Y) * MetersToUnrealUnits
+					);
+
+					FLiveLinkFrameDataStruct TransformStruct(FLiveLinkTransformFrameData::StaticStruct());
+					FLiveLinkTransformFrameData* TransformFrameData = TransformStruct.Cast<FLiveLinkTransformFrameData>();
+
+					TransformFrameData->WorldTime = TimeNowSec;
+					if (CurrentFrameTime)
+					{
+						TransformFrameData->MetaData.SceneTime = *CurrentFrameTime;
+					}
+
+					TransformFrameData->Transform = PoseTransform;
+
+					Send(MoveTemp(TransformStruct), SubjectName);
+				}
+			}
+
+			// Enumerate actions.
+			FLiveLinkFrameDataStruct InputStruct(FLiveLinkGamepadInputDeviceFrameData::StaticStruct());
+			FLiveLinkGamepadInputDeviceFrameData* InputFrameData = InputStruct.Cast<FLiveLinkGamepadInputDeviceFrameData>();
+
+			InputFrameData->WorldTime = TimeNowSec;
+			if (CurrentFrameTime)
+			{
+				InputFrameData->MetaData.SceneTime = *CurrentFrameTime;
+			}
+
+			for (FOpenVRInputAction& Action : Actions)
+			{
+				// Populate action data.
+				switch (Action.ActionType)
+				{
+					case FOpenVRInputAction::Digital:
+					{
+						const vr::VRInputValueHandle_t Unrestricted = vr::k_ulInvalidInputValueHandle;
+						InputError = VrInput->GetDigitalActionData(Action.Handle,
+							&Action.LastActionData.Digital,
+							sizeof(Action.LastActionData.Digital),
+							Unrestricted);
+						break;
+					}
+
+					case FOpenVRInputAction::Analog1D:
+					case FOpenVRInputAction::Analog2D:
+					{
+						const vr::VRInputValueHandle_t Unrestricted = vr::k_ulInvalidInputValueHandle;
+						InputError = VrInput->GetAnalogActionData(Action.Handle,
+							&Action.LastActionData.Analog,
+							sizeof(Action.LastActionData.Analog),
+							Unrestricted);
+						break;
+					}
+
+					default:
+						checkNoEntry();
+				}
+
+				if (InputError != vr::VRInputError_None)
 				{
+					UE_LOGFMT(LogLiveLinkOpenVR, Error,
+						"IVRInput::Get*ActionData for '{ActionName}' failed with result {InputError}",
+						Action.ActionName, InputError);
 					continue;
 				}
 
-				// Transpose and decompose.
-				const FMatrix PoseMatrix = ToFMatrix(Pose.mDeviceToAbsoluteTracking);
-				const FQuat PoseOrientation(PoseMatrix);
-				const FVector PosePosition(PoseMatrix.M[3][0], PoseMatrix.M[3][1], PoseMatrix.M[3][2]);
+				switch (Action.ActionType)
+				{
+					case FOpenVRInputAction::Digital:
+					{
+						float& DestField = InputFrameData->**Action.OutputFieldX;
+						DestField = Action.LastActionData.Digital.bState ? 1.0f : 0.0f;
+						break;
+					}
 
-				// Handedness/basis change/scale.
-				const FTransform PoseTransform(
-					FQuat(-PoseOrientation.Z, PoseOrientation.X, PoseOrientation.Y, -PoseOrientation.W),
-					FVector(-PosePosition.Z, PosePosition.X, PosePosition.Y) * 100.0f
-				);
+					case FOpenVRInputAction::Analog1D:
+					{
+						float& DestFieldX = InputFrameData->**Action.OutputFieldX;
+						DestFieldX = Action.LastActionData.Analog.x;
+						break;
+					}
 
-				FLiveLinkFrameDataStruct FrameData(FLiveLinkTransformFrameData::StaticStruct());
-				FLiveLinkTransformFrameData* TransformFrameData = FrameData.Cast<FLiveLinkTransformFrameData>();
-				TransformFrameData->Transform = PoseTransform;
+					case FOpenVRInputAction::Analog2D:
+					{
+						float& DestFieldX = InputFrameData->**Action.OutputFieldX;
+						float& DestFieldY = InputFrameData->**Action.OutputFieldY;
+						DestFieldX = Action.LastActionData.Analog.x;
+						DestFieldY = Action.LastActionData.Analog.y;
+						break;
+					}
 
-				Send(&FrameData, SubjectName);
+					default:
+						checkNoEntry();
+				}
 			}
+
+			Send(MoveTemp(InputStruct), InputSubjectName);
 		}
 
 		FPlatformProcess::Sleep(0.001f);
@@ -255,14 +487,14 @@ uint32 FLiveLinkOpenVRSource::Run()
 }
 
 
-void FLiveLinkOpenVRSource::Send(FLiveLinkFrameDataStruct* FrameDataToSend, FName SubjectName)
+void FLiveLinkOpenVRSource::Send(FLiveLinkFrameDataStruct&& InFrameData, FName InSubjectName)
 {
 	if (bStopping || (Client == nullptr))
 	{
 		return;
 	}
 
-	Client->PushSubjectFrameData_AnyThread({ SourceGuid, SubjectName }, MoveTemp(*FrameDataToSend));
+	Client->PushSubjectFrameData_AnyThread({ SourceGuid, InSubjectName }, MoveTemp(InFrameData));
 }
 
 
diff --git a/Engine/Plugins/Experimental/LiveLinkOpenVR/Source/LiveLinkOpenVR/Private/LiveLinkOpenVRSource.h b/Engine/Plugins/Experimental/LiveLinkOpenVR/Source/LiveLinkOpenVR/Private/LiveLinkOpenVRSource.h
--- a/Engine/Plugins/Experimental/LiveLinkOpenVR/Source/LiveLinkOpenVR/Private/LiveLinkOpenVRSource.h
+++ b/Engine/Plugins/Experimental/LiveLinkOpenVR/Source/LiveLinkOpenVR/Private/LiveLinkOpenVRSource.h
@@ -8,7 +8,6 @@
 #include "LiveLinkOpenVRTypes.h"
 
 
-
 class FLiveLinkOpenVRSource
 	: public ILiveLinkSource
 	, public FRunnable
@@ -21,7 +20,7 @@ public:
 
 	//~ Begin ILiveLinkSource interface
 	virtual void ReceiveClient(ILiveLinkClient* InClient, FGuid InSourceGuid) override;
-	virtual void InitializeSettings(ULiveLinkSourceSettings* Settings) override;
+	virtual void InitializeSettings(ULiveLinkSourceSettings* InSettings) override;
 	virtual void Update() override { }
 
 	virtual bool IsSourceStillValid() const override;
@@ -43,7 +42,7 @@ public:
 	virtual void Exit() override { }
 	//~ End FRunnable interface
 
-	void Send(FLiveLinkFrameDataStruct* FrameDataToSend, FName SubjectName);
+	void Send(FLiveLinkFrameDataStruct&& InFrameData, FName InSubjectName);
 
 private:
 	// Callback when the a livelink subject has been added
@@ -81,8 +80,13 @@ private:
 	// frame counter for data
 	int64 FrameCounter = 0;
 
-	// Update rate (in Hz) at which to read the tracking data for each device
-	std::atomic<uint32> LocalUpdateRateInHz;
+	TWeakObjectPtr<ULiveLinkOpenVRSourceSettings> Settings;
+
+	// Atomic copies of UObject fields updated on change notification
+	std::atomic<bool> bTrackTrackers_AnyThread;
+	std::atomic<bool> bTrackControllers_AnyThread;
+	std::atomic<bool> bTrackHMDs_AnyThread;
+	std::atomic<uint32> LocalUpdateRateInHz_AnyThread;
 
 	// Delegate for when the LiveLink client has ticked
 	FDelegateHandle OnSubjectAddedDelegate;
diff --git a/Engine/Plugins/Experimental/LiveLinkOpenVR/Source/LiveLinkOpenVR/Private/LiveLinkOpenVRTypes.h b/Engine/Plugins/Experimental/LiveLinkOpenVR/Source/LiveLinkOpenVR/Private/LiveLinkOpenVRTypes.h
--- a/Engine/Plugins/Experimental/LiveLinkOpenVR/Source/LiveLinkOpenVR/Private/LiveLinkOpenVRTypes.h
+++ b/Engine/Plugins/Experimental/LiveLinkOpenVR/Source/LiveLinkOpenVR/Private/LiveLinkOpenVRTypes.h
@@ -13,18 +13,6 @@ struct FLiveLinkOpenVRCommonSettings
 	GENERATED_BODY()
 
 public:
-	/** Track all tracker pucks */
-	UPROPERTY(EditAnywhere, Category="LiveLinkOpenVR Settings")
-	bool bTrackTrackers = true;
-
-	/** Track all controllers */
-	UPROPERTY(EditAnywhere, Category="LiveLinkOpenVR Settings")
-	bool bTrackControllers = false;
-
-	/** Track all HMDs */
-	UPROPERTY(EditAnywhere, Category="LiveLinkOpenVR Settings")
-	bool bTrackHMDs = false;
-
 	/** Update rate (in Hz) at which to read the tracking data for each device */
 	UPROPERTY(EditAnywhere, Category="LiveLinkOpenVR Settings", meta=(ClampMin=1, ClampMax=1000))
 	uint32 LocalUpdateRateInHz = 60;
@@ -38,6 +26,22 @@ struct FLiveLinkOpenVRConnectionSettings
 
 	UPROPERTY(EditAnywhere, Category="LiveLinkOpenVR Settings", meta=(ShowOnlyInnerProperties))
 	FLiveLinkOpenVRCommonSettings CommonSettings;
+
+	/** Track all tracker pucks */
+	UPROPERTY(EditAnywhere, Category="LiveLinkOpenVR Settings")
+	bool bTrackTrackers = true;
+
+	/** Track all tracking references (e.g. base stations) */
+	UPROPERTY(EditAnywhere, Category="LiveLinkOpenVR Settings")
+	bool bTrackTrackingReferences = true;
+
+	/** Track all controllers */
+	UPROPERTY(EditAnywhere, Category="LiveLinkOpenVR Settings")
+	bool bTrackControllers = false;
+
+	/** Track all HMDs */
+	UPROPERTY(EditAnywhere, Category="LiveLinkOpenVR Settings")
+	bool bTrackHMDs = false;
 };
 
 
