LiveLinkOpenVR: Improved default bindings.

Bundles LeftStick(Up|Right) and RightStick(Up|Right) together into new 2D actions, and moves the thumbstick bindings to them. LeftAnalog and RightAnalog now map to trackpad position bindings instead, and SpecialLeft/SpecialRight map to trackpad force.

Also added mappings for grip trigger pull -> shoulder, and face buttons.

diff --git a/Engine/Plugins/Experimental/LiveLinkOpenVR/Config/livelinkopenvr_action_manifest.json b/Engine/Plugins/Experimental/LiveLinkOpenVR/Config/livelinkopenvr_action_manifest.json
--- a/Engine/Plugins/Experimental/LiveLinkOpenVR/Config/livelinkopenvr_action_manifest.json
+++ b/Engine/Plugins/Experimental/LiveLinkOpenVR/Config/livelinkopenvr_action_manifest.json
@@ -91,33 +91,25 @@
 			"type": "vector1"
 		},
 		{
-			"name": "/actions/LiveLinkGamepadInputDevice/in/LeftStickUp",
-			"type": "vector1"
+			"name": "/actions/LiveLinkGamepadInputDevice/in/LeftStick_2D",
+			"type": "vector2"
 		},
 		{
 			"name": "/actions/LiveLinkGamepadInputDevice/in/LeftStickDown",
 			"type": "vector1"
 		},
-		{
-			"name": "/actions/LiveLinkGamepadInputDevice/in/LeftStickRight",
-			"type": "vector1"
-		},
 		{
 			"name": "/actions/LiveLinkGamepadInputDevice/in/LeftStickLeft",
 			"type": "vector1"
 		},
 		{
-			"name": "/actions/LiveLinkGamepadInputDevice/in/RightStickUp",
-			"type": "vector1"
+			"name": "/actions/LiveLinkGamepadInputDevice/in/RightStick_2D",
+			"type": "vector2"
 		},
 		{
 			"name": "/actions/LiveLinkGamepadInputDevice/in/RightStickDown",
 			"type": "vector1"
 		},
-		{
-			"name": "/actions/LiveLinkGamepadInputDevice/in/RightStickRight",
-			"type": "vector1"
-		},
 		{
 			"name": "/actions/LiveLinkGamepadInputDevice/in/RightStickLeft",
 			"type": "vector1"
diff --git a/Engine/Plugins/Experimental/LiveLinkOpenVR/Config/livelinkopenvr_bindings_knuckles.json b/Engine/Plugins/Experimental/LiveLinkOpenVR/Config/livelinkopenvr_bindings_knuckles.json
--- a/Engine/Plugins/Experimental/LiveLinkOpenVR/Config/livelinkopenvr_bindings_knuckles.json
+++ b/Engine/Plugins/Experimental/LiveLinkOpenVR/Config/livelinkopenvr_bindings_knuckles.json
@@ -10,7 +10,7 @@
 					"mode": "joystick",
 					"inputs": {
 						"position": {
-							"output": "/actions/LiveLinkGamepadInputDevice/in/LeftAnalog_2D"
+							"output": "/actions/LiveLinkGamepadInputDevice/in/LeftStick_2D"
 						},
 						"click": {
 							"output": "/actions/LiveLinkGamepadInputDevice/in/LeftThumb"
@@ -31,19 +31,52 @@
 					"mode": "trackpad",
 					"inputs": {
 						"position": {
-							"output": "/actions/LiveLinkGamepadInputDevice/in/SpecialLeft_2D"
-						},
-						"click": {
+							"output": "/actions/LiveLinkGamepadInputDevice/in/LeftAnalog_2D"
+						}
+					}
+				},
+				{
+					"path": "/user/hand/left/input/trackpad",
+					"mode": "force_sensor",
+					"inputs": {
+						"force": {
 							"output": "/actions/LiveLinkGamepadInputDevice/in/SpecialLeft"
 						}
 					}
 				},
+				{
+					"path": "/user/hand/left/input/a",
+					"mode": "button",
+					"inputs": {
+						"click": {
+							"output": "/actions/LiveLinkGamepadInputDevice/in/FaceButtonLeft"
+						}
+					}
+				},
+				{
+					"path": "/user/hand/left/input/b",
+					"mode": "button",
+					"inputs": {
+						"click": {
+							"output": "/actions/LiveLinkGamepadInputDevice/in/FaceButtonTop"
+						}
+					}
+				},
+				{
+					"path": "/user/hand/left/input/grip",
+					"mode": "trigger",
+					"inputs": {
+						"pull": {
+							"output": "/actions/LiveLinkGamepadInputDevice/in/LeftShoulder"
+						}
+					}
+				},
 				{
 					"path": "/user/hand/right/input/thumbstick",
 					"mode": "joystick",
 					"inputs": {
 						"position": {
-							"output": "/actions/LiveLinkGamepadInputDevice/in/RightAnalog_2D"
+							"output": "/actions/LiveLinkGamepadInputDevice/in/RightStick_2D"
 						},
 						"click": {
 							"output": "/actions/LiveLinkGamepadInputDevice/in/RightThumb"
@@ -58,6 +91,51 @@
 							"output": "/actions/LiveLinkGamepadInputDevice/in/RightTriggerAnalog"
 						}
 					}
+				},
+				{
+					"path": "/user/hand/right/input/trackpad",
+					"mode": "trackpad",
+					"inputs": {
+						"position": {
+							"output": "/actions/LiveLinkGamepadInputDevice/in/RightAnalog_2D"
+						}
+					}
+				},
+				{
+					"path": "/user/hand/right/input/trackpad",
+					"mode": "force_sensor",
+					"inputs": {
+						"force": {
+							"output": "/actions/LiveLinkGamepadInputDevice/in/SpecialRight"
+						}
+					}
+				},
+				{
+					"path": "/user/hand/right/input/a",
+					"mode": "button",
+					"inputs": {
+						"click": {
+							"output": "/actions/LiveLinkGamepadInputDevice/in/FaceButtonBottom"
+						}
+					}
+				},
+				{
+					"path": "/user/hand/right/input/b",
+					"mode": "button",
+					"inputs": {
+						"click": {
+							"output": "/actions/LiveLinkGamepadInputDevice/in/FaceButtonRight"
+						}
+					}
+				},
+				{
+					"path": "/user/hand/right/input/grip",
+					"mode": "trigger",
+					"inputs": {
+						"pull": {
+							"output": "/actions/LiveLinkGamepadInputDevice/in/RightShoulder"
+						}
+					}
 				}
 			]
 		}
diff --git a/Engine/Plugins/Experimental/LiveLinkOpenVR/Source/LiveLinkOpenVR/Private/LiveLinkOpenVRSource.cpp b/Engine/Plugins/Experimental/LiveLinkOpenVR/Source/LiveLinkOpenVR/Private/LiveLinkOpenVRSource.cpp
--- a/Engine/Plugins/Experimental/LiveLinkOpenVR/Source/LiveLinkOpenVR/Private/LiveLinkOpenVRSource.cpp
+++ b/Engine/Plugins/Experimental/LiveLinkOpenVR/Source/LiveLinkOpenVR/Private/LiveLinkOpenVRSource.cpp
@@ -190,7 +190,17 @@ uint32 FLiveLinkOpenVRSource::Run()
 		{ "SpecialLeft_2D", FOpenVRInputAction::Analog2D,
 		  &FLiveLinkGamepadInputDeviceFrameData::SpecialLeft_X, &FLiveLinkGamepadInputDeviceFrameData::SpecialLeft_Y },
 
-#define DECLARE_VECTOR1_ACTION(ActionNameToken)                     \
+		{ "LeftStick_2D", FOpenVRInputAction::Analog2D,
+		  &FLiveLinkGamepadInputDeviceFrameData::LeftStickRight, &FLiveLinkGamepadInputDeviceFrameData::LeftStickUp },
+
+		{ "RightStick_2D", FOpenVRInputAction::Analog2D,
+		  &FLiveLinkGamepadInputDeviceFrameData::RightStickRight, &FLiveLinkGamepadInputDeviceFrameData::RightStickUp },
+
+		// All actions delcared in the manifest are listed below.
+		// Most are 1D/float actions, but the few that are commented
+		// out are the ones that have been combined into 2D axes above.
+
+#define DECLARE_VECTOR1_ACTION(ActionNameToken)                        \
 		{ UE_STRINGIZE(ActionNameToken), FOpenVRInputAction::Analog1D, \
 		  &FLiveLinkGamepadInputDeviceFrameData::ActionNameToken },
 
@@ -218,13 +228,13 @@ uint32 FLiveLinkOpenVRSource::Run()
 		DECLARE_VECTOR1_ACTION(DPadDown)
 		DECLARE_VECTOR1_ACTION(DPadRight)
 		DECLARE_VECTOR1_ACTION(DPadLeft)
-		DECLARE_VECTOR1_ACTION(LeftStickUp)
+		//DECLARE_VECTOR1_ACTION(LeftStickUp)
 		DECLARE_VECTOR1_ACTION(LeftStickDown)
-		DECLARE_VECTOR1_ACTION(LeftStickRight)
+		//DECLARE_VECTOR1_ACTION(LeftStickRight)
 		DECLARE_VECTOR1_ACTION(LeftStickLeft)
-		DECLARE_VECTOR1_ACTION(RightStickUp)
+		//DECLARE_VECTOR1_ACTION(RightStickUp)
 		DECLARE_VECTOR1_ACTION(RightStickDown)
-		DECLARE_VECTOR1_ACTION(RightStickRight)
+		//DECLARE_VECTOR1_ACTION(RightStickRight)
 		DECLARE_VECTOR1_ACTION(RightStickLeft)
 
 #undef DECLARE_VECTOR1_ACTION
